Add Physics asset toolset.

#Synced-CL 52186709

diff --git a/Engine/Plugins/Experimental/Toolsets/PhysicsToolsets/PhysicsToolsets.uplugin b/Engine/Plugins/Experimental/Toolsets/PhysicsToolsets/PhysicsToolsets.uplugin
new file mode 100644
--- /dev/null
+++ b/Engine/Plugins/Experimental/Toolsets/PhysicsToolsets/PhysicsToolsets.uplugin
@@ -0,0 +1,34 @@
+{
+	"FileVersion": 3,
+	"Version": 1,
+	"VersionName": "1.0",
+	"FriendlyName": "Physics Toolsets",
+	"Description": "AI-callable toolsets for working with Unreal Engine physics assets.",
+	"Category": "Other",
+	"EditorOnly": true,
+	"CreatedBy": "",
+	"CreatedByURL": "",
+	"DocsURL": "",
+	"MarketplaceURL": "",
+	"SupportURL": "",
+	"CanContainContent": false,
+	"IsBetaVersion": false,
+	"IsExperimentalVersion": true,
+	"NoRedist": true,
+	"Installed": false,
+	"EnabledByDefault": false,
+	"Modules": [
+		{
+			"Name": "PhysicsToolsets",
+			"Type": "Editor",
+			"LoadingPhase": "PostEngineInit",
+			"TargetAllowList": ["Editor"]
+		}
+	],
+	"Plugins": [
+		{
+			"Name": "ToolsetRegistry",
+			"Enabled": true
+		}
+	]
+}
diff --git a/Engine/Plugins/Experimental/Toolsets/PhysicsToolsets/Source/PhysicsToolsets/PhysicsToolsets.Build.cs b/Engine/Plugins/Experimental/Toolsets/PhysicsToolsets/Source/PhysicsToolsets/PhysicsToolsets.Build.cs
new file mode 100644
--- /dev/null
+++ b/Engine/Plugins/Experimental/Toolsets/PhysicsToolsets/Source/PhysicsToolsets/PhysicsToolsets.Build.cs
@@ -0,0 +1,32 @@
+// Copyright Epic Games, Inc. All Rights Reserved.
+
+using UnrealBuildTool;
+
+public class PhysicsToolsets : ModuleRules
+{
+	public PhysicsToolsets(ReadOnlyTargetRules Target) : base(Target)
+	{
+		PCHUsage = ModuleRules.PCHUsageMode.UseExplicitOrSharedPCHs;
+		IWYUSupport = IWYUSupport.Full;
+		bUseUnity = true;
+
+		PublicDependencyModuleNames.AddRange(
+			new string[]
+			{
+				"Core",
+			}
+		);
+
+		PrivateDependencyModuleNames.AddRange(
+			new string[]
+			{
+				"AssetRegistry",
+				"CoreUObject",
+				"Engine",
+				"PhysicsUtilities",
+				"ToolsetRegistry",
+				"UnrealEd",
+			}
+		);
+	}
+}
diff --git a/Engine/Plugins/Experimental/Toolsets/PhysicsToolsets/Source/PhysicsToolsets/Private/PhysicsToolsets/Module.cpp b/Engine/Plugins/Experimental/Toolsets/PhysicsToolsets/Source/PhysicsToolsets/Private/PhysicsToolsets/Module.cpp
new file mode 100644
--- /dev/null
+++ b/Engine/Plugins/Experimental/Toolsets/PhysicsToolsets/Source/PhysicsToolsets/Private/PhysicsToolsets/Module.cpp
@@ -0,0 +1,25 @@
+// Copyright Epic Games, Inc. All Rights Reserved.
+
+#include "PhysicsToolsets/Module.h"
+
+#include "Modules/ModuleManager.h"
+#include "ToolsetRegistry/UToolsetRegistry.h"
+
+#include "PhysicsToolsets/PhysicsAssetToolset.h"
+
+#define LOCTEXT_NAMESPACE "FPhysicsToolsetsModule"
+
+void FPhysicsToolsetsModule::StartupModule()
+{
+	UToolsetRegistry::RegisterToolsetClass(UPhysicsAssetToolset::StaticClass());
+}
+
+void FPhysicsToolsetsModule::ShutdownModule()
+{
+	UToolsetRegistry::UnregisterToolsetClass(UPhysicsAssetToolset::StaticClass());
+}
+
+#undef LOCTEXT_NAMESPACE
+
+IMPLEMENT_MODULE(FPhysicsToolsetsModule, PhysicsToolsets);
+DEFINE_LOG_CATEGORY(LogPhysicsToolsets);
diff --git a/Engine/Plugins/Experimental/Toolsets/PhysicsToolsets/Source/PhysicsToolsets/Private/PhysicsToolsets/Module.h b/Engine/Plugins/Experimental/Toolsets/PhysicsToolsets/Source/PhysicsToolsets/Private/PhysicsToolsets/Module.h
new file mode 100644
--- /dev/null
+++ b/Engine/Plugins/Experimental/Toolsets/PhysicsToolsets/Source/PhysicsToolsets/Private/PhysicsToolsets/Module.h
@@ -0,0 +1,15 @@
+// Copyright Epic Games, Inc. All Rights Reserved.
+
+#pragma once
+
+#include "Modules/ModuleInterface.h"
+#include "Logging/LogMacros.h"
+
+DECLARE_LOG_CATEGORY_EXTERN(LogPhysicsToolsets, Log, Log);
+
+class FPhysicsToolsetsModule : public IModuleInterface
+{
+public:
+	virtual void StartupModule() override;
+	virtual void ShutdownModule() override;
+};
diff --git a/Engine/Plugins/Experimental/Toolsets/PhysicsToolsets/Source/PhysicsToolsets/Private/PhysicsToolsets/PhysicsAssetToolset.cpp b/Engine/Plugins/Experimental/Toolsets/PhysicsToolsets/Source/PhysicsToolsets/Private/PhysicsToolsets/PhysicsAssetToolset.cpp
new file mode 100644
--- /dev/null
+++ b/Engine/Plugins/Experimental/Toolsets/PhysicsToolsets/Source/PhysicsToolsets/Private/PhysicsToolsets/PhysicsAssetToolset.cpp
@@ -0,0 +1,695 @@
+// Copyright Epic Games, Inc. All Rights Reserved.
+
+#include "PhysicsToolsets/PhysicsAssetToolset.h"
+
+#include "AssetRegistry/AssetRegistryModule.h"
+#include "Chaos/ChaosEngineInterface.h"
+#include "Engine/SkeletalMesh.h"
+#include "Kismet/KismetSystemLibrary.h"
+#include "Misc/PackageName.h"
+#include "PhysicsAssetUtils.h"
+#include "PhysicsEngine/BoxElem.h"
+#include "PhysicsEngine/ConstraintInstance.h"
+#include "PhysicsEngine/PhysicsAsset.h"
+#include "PhysicsEngine/PhysicsConstraintTemplate.h"
+#include "PhysicsEngine/SkeletalBodySetup.h"
+#include "PhysicsEngine/SphereElem.h"
+#include "PhysicsEngine/SphylElem.h"
+#include "UObject/Package.h"
+#include "UObject/UObjectGlobals.h"
+
+#include UE_INLINE_GENERATED_CPP_BY_NAME(PhysicsAssetToolset)
+
+// ---------------------------------------------------------------------------
+// File-local helpers
+// ---------------------------------------------------------------------------
+
+namespace
+{
+	EAngularConstraintMotion ToAngularMotion(EConstraintMotion Motion)
+	{
+		switch (Motion)
+		{
+		case EConstraintMotion::Limited: return ACM_Limited;
+		case EConstraintMotion::Locked:  return ACM_Locked;
+		default:                         return ACM_Free;
+		}
+	}
+
+	EConstraintMotion FromAngularMotion(EAngularConstraintMotion Motion)
+	{
+		switch (Motion)
+		{
+		case ACM_Limited: return EConstraintMotion::Limited;
+		case ACM_Locked:  return EConstraintMotion::Locked;
+		default:          return EConstraintMotion::Free;
+		}
+	}
+} // namespace
+
+// ---------------------------------------------------------------------------
+// Public API — shape primitives
+// ---------------------------------------------------------------------------
+
+UPhysicsAsset* UPhysicsAssetToolset::CreateFromMesh(
+	const FString& MeshPath, bool bAssignToMesh)
+{
+	if (MeshPath.IsEmpty())
+	{
+		UKismetSystemLibrary::RaiseScriptError(TEXT("MeshPath cannot be empty."));
+		return nullptr;
+	}
+
+	USkeletalMesh* Mesh = Cast<USkeletalMesh>(
+		StaticLoadObject(USkeletalMesh::StaticClass(), nullptr, *MeshPath));
+	if (!Mesh)
+	{
+		UKismetSystemLibrary::RaiseScriptError(
+			FString::Printf(TEXT("No skeletal mesh found at '%s'."), *MeshPath));
+		return nullptr;
+	}
+
+	const FString PackageDir = FPackageName::GetLongPackagePath(MeshPath);
+	const FString AssetName = FPackageName::GetShortName(MeshPath) + TEXT("_PhysicsAsset");
+	const FString NewPackagePath = PackageDir / AssetName;
+
+	UPackage* NewPackage = CreatePackage(*NewPackagePath);
+	if (!NewPackage)
+	{
+		UKismetSystemLibrary::RaiseScriptError(
+			FString::Printf(TEXT("Failed to create package at '%s'."), *NewPackagePath));
+		return nullptr;
+	}
+
+	UPhysicsAsset* PhysicsAsset = NewObject<UPhysicsAsset>(
+		NewPackage, FName(*AssetName), RF_Public | RF_Standalone | RF_Transactional);
+
+	FPhysAssetCreateParams Params;
+	FText ErrorMessage;
+	const bool bSuccess = FPhysicsAssetUtils::CreateFromSkeletalMesh(
+		PhysicsAsset, Mesh, Params, ErrorMessage, bAssignToMesh, /*bShowProgress=*/false);
+
+	if (!bSuccess)
+	{
+		UKismetSystemLibrary::RaiseScriptError(
+			FString::Printf(TEXT("Failed to create physics asset for '%s': %s"),
+				*MeshPath, *ErrorMessage.ToString()));
+		return nullptr;
+	}
+
+	FAssetRegistryModule::AssetCreated(PhysicsAsset);
+	PhysicsAsset->MarkPackageDirty();
+	return PhysicsAsset;
+}
+
+TArray<FString> UPhysicsAssetToolset::GetBodyNames(UPhysicsAsset* PhysicsAsset)
+{
+	if (!PhysicsAsset)
+	{
+		UKismetSystemLibrary::RaiseScriptError(TEXT("PhysicsAsset is null."));
+		return {};
+	}
+
+	TArray<FString> Names;
+	Names.Reserve(PhysicsAsset->SkeletalBodySetups.Num());
+	for (const TObjectPtr<USkeletalBodySetup>& Body : PhysicsAsset->SkeletalBodySetups)
+	{
+		if (Body)
+		{
+			Names.Add(Body->BoneName.ToString());
+		}
+	}
+	return Names;
+}
+
+TArray<FPhysicsShapeInfo> UPhysicsAssetToolset::GetBodyShapes(
+	UPhysicsAsset* PhysicsAsset, const FString& BoneName)
+{
+	if (!PhysicsAsset)
+	{
+		UKismetSystemLibrary::RaiseScriptError(TEXT("PhysicsAsset is null."));
+		return {};
+	}
+
+	USkeletalBodySetup* Body = FindBody(PhysicsAsset, BoneName);
+	if (!Body)
+	{
+		UKismetSystemLibrary::RaiseScriptError(
+			FString::Printf(TEXT("No body found for bone '%s'."), *BoneName));
+		return {};
+	}
+
+	TArray<FPhysicsShapeInfo> Shapes;
+
+	for (const FKSphereElem& Elem : Body->AggGeom.SphereElems)
+	{
+		FPhysicsShapeInfo& Info = Shapes.AddDefaulted_GetRef();
+		Info.ShapeName = Elem.GetName().ToString();
+		Info.ShapeType = EPhysicsShapeType::Sphere;
+		Info.Center = Elem.Center;
+		Info.Radius = Elem.Radius;
+	}
+
+	for (const FKSphylElem& Elem : Body->AggGeom.SphylElems)
+	{
+		FPhysicsShapeInfo& Info = Shapes.AddDefaulted_GetRef();
+		Info.ShapeName = Elem.GetName().ToString();
+		Info.ShapeType = EPhysicsShapeType::Capsule;
+		Info.Center = Elem.Center;
+		Info.Rotation = Elem.Rotation;
+		Info.Radius = Elem.Radius;
+		Info.Length = Elem.Length;
+	}
+
+	for (const FKBoxElem& Elem : Body->AggGeom.BoxElems)
+	{
+		FPhysicsShapeInfo& Info = Shapes.AddDefaulted_GetRef();
+		Info.ShapeName = Elem.GetName().ToString();
+		Info.ShapeType = EPhysicsShapeType::Box;
+		Info.Center = Elem.Center;
+		Info.Rotation = Elem.Rotation;
+		Info.ExtentX = Elem.X;
+		Info.ExtentY = Elem.Y;
+		Info.ExtentZ = Elem.Z;
+	}
+
+	return Shapes;
+}
+
+void UPhysicsAssetToolset::SetSphere(
+	UPhysicsAsset* PhysicsAsset, const FString& BoneName,
+	const FString& ShapeName, const FVector& Center, float Radius)
+{
+	USkeletalBodySetup* Body = FindBodyForShape(PhysicsAsset, BoneName, ShapeName);
+	if (!Body) return;
+	if (Radius <= 0.f)
+	{
+		UKismetSystemLibrary::RaiseScriptError(TEXT("Radius must be greater than zero."));
+		return;
+	}
+
+	const FName Name = BeginShapeEdit(PhysicsAsset, Body, ShapeName);
+
+	FKSphereElem& NewElem = Body->AggGeom.SphereElems.AddDefaulted_GetRef();
+	NewElem.SetName(Name);
+	NewElem.Center = Center;
+	NewElem.Radius = Radius;
+
+	NotifyAssetChanged(PhysicsAsset);
+}
+
+void UPhysicsAssetToolset::SetCapsule(
+	UPhysicsAsset* PhysicsAsset, const FString& BoneName,
+	const FString& ShapeName, const FVector& Center, const FRotator& Rotation,
+	float Radius, float Length)
+{
+	USkeletalBodySetup* Body = FindBodyForShape(PhysicsAsset, BoneName, ShapeName);
+	if (!Body) return;
+	if (Radius <= 0.f)
+	{
+		UKismetSystemLibrary::RaiseScriptError(TEXT("Radius must be greater than zero."));
+		return;
+	}
+	if (Length < 0.f)
+	{
+		UKismetSystemLibrary::RaiseScriptError(TEXT("Length must be non-negative."));
+		return;
+	}
+
+	const FName Name = BeginShapeEdit(PhysicsAsset, Body, ShapeName);
+
+	FKSphylElem& NewElem = Body->AggGeom.SphylElems.AddDefaulted_GetRef();
+	NewElem.SetName(Name);
+	NewElem.Center = Center;
+	NewElem.Rotation = Rotation;
+	NewElem.Radius = Radius;
+	NewElem.Length = Length;
+
+	NotifyAssetChanged(PhysicsAsset);
+}
+
+void UPhysicsAssetToolset::SetBox(
+	UPhysicsAsset* PhysicsAsset, const FString& BoneName,
+	const FString& ShapeName, const FVector& Center, const FRotator& Rotation,
+	float ExtentX, float ExtentY, float ExtentZ)
+{
+	USkeletalBodySetup* Body = FindBodyForShape(PhysicsAsset, BoneName, ShapeName);
+	if (!Body) return;
+	if (ExtentX <= 0.f || ExtentY <= 0.f || ExtentZ <= 0.f)
+	{
+		UKismetSystemLibrary::RaiseScriptError(
+			TEXT("ExtentX, ExtentY, and ExtentZ must all be greater than zero."));
+		return;
+	}
+
+	const FName Name = BeginShapeEdit(PhysicsAsset, Body, ShapeName);
+
+	FKBoxElem& NewElem = Body->AggGeom.BoxElems.AddDefaulted_GetRef();
+	NewElem.SetName(Name);
+	NewElem.Center = Center;
+	NewElem.Rotation = Rotation;
+	NewElem.X = ExtentX;
+	NewElem.Y = ExtentY;
+	NewElem.Z = ExtentZ;
+
+	NotifyAssetChanged(PhysicsAsset);
+}
+
+void UPhysicsAssetToolset::RemoveShape(
+	UPhysicsAsset* PhysicsAsset, const FString& BoneName, const FString& ShapeName)
+{
+	USkeletalBodySetup* Body = FindBodyForShape(PhysicsAsset, BoneName, ShapeName);
+	if (!Body) return;
+
+	const int32 Before =
+		Body->AggGeom.SphereElems.Num() +
+		Body->AggGeom.SphylElems.Num() +
+		Body->AggGeom.BoxElems.Num();
+
+	const FName Name(*ShapeName);
+	PhysicsAsset->Modify();
+	Body->Modify();
+	RemoveShapeByName(Body, Name);
+
+	const int32 After =
+		Body->AggGeom.SphereElems.Num() +
+		Body->AggGeom.SphylElems.Num() +
+		Body->AggGeom.BoxElems.Num();
+
+	if (Before == After)
+	{
+		UKismetSystemLibrary::RaiseScriptError(
+			FString::Printf(
+				TEXT("No shape named '%s' found on body '%s'."), *ShapeName, *BoneName));
+		return;
+	}
+
+	NotifyAssetChanged(PhysicsAsset);
+}
+
+// ---------------------------------------------------------------------------
+// Public API — body CRUD
+// ---------------------------------------------------------------------------
+
+void UPhysicsAssetToolset::AddBody(UPhysicsAsset* PhysicsAsset, const FString& BoneName)
+{
+	if (!PhysicsAsset)
+	{
+		UKismetSystemLibrary::RaiseScriptError(TEXT("PhysicsAsset is null."));
+		return;
+	}
+	if (BoneName.IsEmpty())
+	{
+		UKismetSystemLibrary::RaiseScriptError(TEXT("BoneName cannot be empty."));
+		return;
+	}
+	if (FindBody(PhysicsAsset, BoneName))
+	{
+		UKismetSystemLibrary::RaiseScriptError(
+			FString::Printf(TEXT("A body for bone '%s' already exists."), *BoneName));
+		return;
+	}
+
+	PhysicsAsset->Modify();
+
+	USkeletalBodySetup* NewBody = NewObject<USkeletalBodySetup>(PhysicsAsset);
+	NewBody->BoneName = FName(*BoneName);
+	PhysicsAsset->SkeletalBodySetups.Add(NewBody);
+
+	NotifyAssetChanged(PhysicsAsset);
+}
+
+void UPhysicsAssetToolset::RemoveBody(UPhysicsAsset* PhysicsAsset, const FString& BoneName)
+{
+	if (!PhysicsAsset)
+	{
+		UKismetSystemLibrary::RaiseScriptError(TEXT("PhysicsAsset is null."));
+		return;
+	}
+
+	const int32 BodyIndex = PhysicsAsset->FindBodyIndex(FName(*BoneName));
+	if (BodyIndex == INDEX_NONE)
+	{
+		UKismetSystemLibrary::RaiseScriptError(
+			FString::Printf(TEXT("No body found for bone '%s'."), *BoneName));
+		return;
+	}
+
+	PhysicsAsset->Modify();
+
+	// Remove all constraints that reference this body (highest index first).
+	TArray<int32> ConstraintIndices;
+	PhysicsAsset->BodyFindConstraints(BodyIndex, ConstraintIndices);
+	ConstraintIndices.Sort([](int32 A, int32 B) { return A > B; });
+	for (const int32 Idx : ConstraintIndices)
+	{
+		PhysicsAsset->ConstraintSetup[Idx]->Modify();
+		PhysicsAsset->ConstraintSetup.RemoveAt(Idx);
+	}
+
+	PhysicsAsset->SkeletalBodySetups[BodyIndex]->Modify();
+	PhysicsAsset->SkeletalBodySetups.RemoveAt(BodyIndex);
+
+	NotifyAssetChanged(PhysicsAsset);
+}
+
+// ---------------------------------------------------------------------------
+// Public API — body properties
+// ---------------------------------------------------------------------------
+
+EBodyPhysicsMode UPhysicsAssetToolset::GetBodyPhysicsMode(
+	UPhysicsAsset* PhysicsAsset, const FString& BoneName)
+{
+	if (!PhysicsAsset)
+	{
+		UKismetSystemLibrary::RaiseScriptError(TEXT("PhysicsAsset is null."));
+		return EBodyPhysicsMode::Default;
+	}
+
+	USkeletalBodySetup* Body = FindBody(PhysicsAsset, BoneName);
+	if (!Body)
+	{
+		UKismetSystemLibrary::RaiseScriptError(
+			FString::Printf(TEXT("No body found for bone '%s'."), *BoneName));
+		return EBodyPhysicsMode::Default;
+	}
+
+	switch (Body->PhysicsType)
+	{
+	case PhysType_Kinematic: return EBodyPhysicsMode::Kinematic;
+	case PhysType_Simulated: return EBodyPhysicsMode::Simulated;
+	default:                 return EBodyPhysicsMode::Default;
+	}
+}
+
+void UPhysicsAssetToolset::SetBodyPhysicsMode(
+	UPhysicsAsset* PhysicsAsset, const FString& BoneName, EBodyPhysicsMode Mode)
+{
+	if (!PhysicsAsset)
+	{
+		UKismetSystemLibrary::RaiseScriptError(TEXT("PhysicsAsset is null."));
+		return;
+	}
+
+	USkeletalBodySetup* Body = FindBody(PhysicsAsset, BoneName);
+	if (!Body)
+	{
+		UKismetSystemLibrary::RaiseScriptError(
+			FString::Printf(TEXT("No body found for bone '%s'."), *BoneName));
+		return;
+	}
+
+	PhysicsAsset->Modify();
+	Body->Modify();
+
+	switch (Mode)
+	{
+	case EBodyPhysicsMode::Kinematic: Body->PhysicsType = PhysType_Kinematic; break;
+	case EBodyPhysicsMode::Simulated: Body->PhysicsType = PhysType_Simulated; break;
+	default:                          Body->PhysicsType = PhysType_Default;   break;
+	}
+
+	NotifyAssetChanged(PhysicsAsset);
+}
+
+float UPhysicsAssetToolset::GetBodyMassScale(
+	UPhysicsAsset* PhysicsAsset, const FString& BoneName)
+{
+	if (!PhysicsAsset)
+	{
+		UKismetSystemLibrary::RaiseScriptError(TEXT("PhysicsAsset is null."));
+		return -1.f;
+	}
+
+	USkeletalBodySetup* Body = FindBody(PhysicsAsset, BoneName);
+	if (!Body)
+	{
+		UKismetSystemLibrary::RaiseScriptError(
+			FString::Printf(TEXT("No body found for bone '%s'."), *BoneName));
+		return -1.f;
+	}
+
+	return Body->DefaultInstance.MassScale;
+}
+
+void UPhysicsAssetToolset::SetBodyMassScale(
+	UPhysicsAsset* PhysicsAsset, const FString& BoneName, float MassScale)
+{
+	if (!PhysicsAsset)
+	{
+		UKismetSystemLibrary::RaiseScriptError(TEXT("PhysicsAsset is null."));
+		return;
+	}
+	if (MassScale <= 0.f)
+	{
+		UKismetSystemLibrary::RaiseScriptError(TEXT("MassScale must be greater than zero."));
+		return;
+	}
+
+	USkeletalBodySetup* Body = FindBody(PhysicsAsset, BoneName);
+	if (!Body)
+	{
+		UKismetSystemLibrary::RaiseScriptError(
+			FString::Printf(TEXT("No body found for bone '%s'."), *BoneName));
+		return;
+	}
+
+	PhysicsAsset->Modify();
+	Body->Modify();
+	Body->DefaultInstance.MassScale = MassScale;
+
+	NotifyAssetChanged(PhysicsAsset);
+}
+
+// ---------------------------------------------------------------------------
+// Public API — constraint CRUD
+// ---------------------------------------------------------------------------
+
+TArray<FPhysicsConstraintInfo> UPhysicsAssetToolset::GetConstraints(
+	UPhysicsAsset* PhysicsAsset)
+{
+	if (!PhysicsAsset)
+	{
+		UKismetSystemLibrary::RaiseScriptError(TEXT("PhysicsAsset is null."));
+		return {};
+	}
+
+	TArray<FPhysicsConstraintInfo> Result;
+	Result.Reserve(PhysicsAsset->ConstraintSetup.Num());
+
+	for (const TObjectPtr<UPhysicsConstraintTemplate>& CT : PhysicsAsset->ConstraintSetup)
+	{
+		if (!CT)
+		{
+			continue;
+		}
+		const FConstraintInstance& CI = CT->DefaultInstance;
+
+		FPhysicsConstraintInfo& Info = Result.AddDefaulted_GetRef();
+		Info.Bone1Name          = CI.ConstraintBone1.ToString();
+		Info.Bone2Name          = CI.ConstraintBone2.ToString();
+		Info.Swing1Motion       = FromAngularMotion(CI.GetAngularSwing1Motion());
+		Info.Swing1LimitDegrees = CI.GetAngularSwing1Limit();
+		Info.Swing2Motion       = FromAngularMotion(CI.GetAngularSwing2Motion());
+		Info.Swing2LimitDegrees = CI.GetAngularSwing2Limit();
+		Info.TwistMotion        = FromAngularMotion(CI.GetAngularTwistMotion());
+		Info.TwistLimitDegrees  = CI.GetAngularTwistLimit();
+	}
+
+	return Result;
+}
+
+void UPhysicsAssetToolset::AddConstraint(
+	UPhysicsAsset* PhysicsAsset, const FString& Bone1Name, const FString& Bone2Name)
+{
+	if (!PhysicsAsset)
+	{
+		UKismetSystemLibrary::RaiseScriptError(TEXT("PhysicsAsset is null."));
+		return;
+	}
+	if (Bone1Name.IsEmpty() || Bone2Name.IsEmpty())
+	{
+		UKismetSystemLibrary::RaiseScriptError(TEXT("Bone names cannot be empty."));
+		return;
+	}
+	if (!FindBody(PhysicsAsset, Bone1Name))
+	{
+		UKismetSystemLibrary::RaiseScriptError(
+			FString::Printf(TEXT("No body found for bone '%s'."), *Bone1Name));
+		return;
+	}
+	if (!FindBody(PhysicsAsset, Bone2Name))
+	{
+		UKismetSystemLibrary::RaiseScriptError(
+			FString::Printf(TEXT("No body found for bone '%s'."), *Bone2Name));
+		return;
+	}
+	if (FindConstraintTemplate(PhysicsAsset, Bone1Name, Bone2Name))
+	{
+		UKismetSystemLibrary::RaiseScriptError(
+			FString::Printf(
+				TEXT("A constraint between '%s' and '%s' already exists."),
+				*Bone1Name, *Bone2Name));
+		return;
+	}
+
+	PhysicsAsset->Modify();
+
+	UPhysicsConstraintTemplate* CT = NewObject<UPhysicsConstraintTemplate>(PhysicsAsset);
+	CT->DefaultInstance.ConstraintBone1 = FName(*Bone1Name);
+	CT->DefaultInstance.ConstraintBone2 = FName(*Bone2Name);
+	PhysicsAsset->ConstraintSetup.Add(CT);
+
+	NotifyAssetChanged(PhysicsAsset);
+}
+
+void UPhysicsAssetToolset::SetConstraintLimits(
+	UPhysicsAsset* PhysicsAsset, FPhysicsConstraintInfo Info)
+{
+	if (!PhysicsAsset)
+	{
+		UKismetSystemLibrary::RaiseScriptError(TEXT("PhysicsAsset is null."));
+		return;
+	}
+
+	UPhysicsConstraintTemplate* CT =
+		FindConstraintTemplate(PhysicsAsset, Info.Bone1Name, Info.Bone2Name);
+	if (!CT)
+	{
+		UKismetSystemLibrary::RaiseScriptError(
+			FString::Printf(
+				TEXT("No constraint found between '%s' and '%s'."),
+				*Info.Bone1Name, *Info.Bone2Name));
+		return;
+	}
+
+	PhysicsAsset->Modify();
+	CT->Modify();
+
+	FConstraintInstance& CI = CT->DefaultInstance;
+	CI.SetAngularSwing1Limit(ToAngularMotion(Info.Swing1Motion), Info.Swing1LimitDegrees);
+	CI.SetAngularSwing2Limit(ToAngularMotion(Info.Swing2Motion), Info.Swing2LimitDegrees);
+	CI.SetAngularTwistLimit(ToAngularMotion(Info.TwistMotion),   Info.TwistLimitDegrees);
+
+	NotifyAssetChanged(PhysicsAsset);
+}
+
+void UPhysicsAssetToolset::RemoveConstraint(
+	UPhysicsAsset* PhysicsAsset, const FString& Bone1Name, const FString& Bone2Name)
+{
+	if (!PhysicsAsset)
+	{
+		UKismetSystemLibrary::RaiseScriptError(TEXT("PhysicsAsset is null."));
+		return;
+	}
+
+	const FName Name1(*Bone1Name);
+	const FName Name2(*Bone2Name);
+
+	int32 Index = PhysicsAsset->FindConstraintIndex(Name1, Name2);
+	if (Index == INDEX_NONE)
+	{
+		Index = PhysicsAsset->FindConstraintIndex(Name2, Name1);
+	}
+	if (Index == INDEX_NONE)
+	{
+		UKismetSystemLibrary::RaiseScriptError(
+			FString::Printf(
+				TEXT("No constraint found between '%s' and '%s'."),
+				*Bone1Name, *Bone2Name));
+		return;
+	}
+
+	PhysicsAsset->Modify();
+	PhysicsAsset->ConstraintSetup[Index]->Modify();
+	PhysicsAsset->ConstraintSetup.RemoveAt(Index);
+
+	NotifyAssetChanged(PhysicsAsset);
+}
+
+// ---------------------------------------------------------------------------
+// Private helpers
+// ---------------------------------------------------------------------------
+
+USkeletalBodySetup* UPhysicsAssetToolset::FindBody(
+	UPhysicsAsset* PhysicsAsset, const FString& BoneName)
+{
+	const int32 Index = PhysicsAsset->FindBodyIndex(FName(*BoneName));
+	if (Index == INDEX_NONE)
+	{
+		return nullptr;
+	}
+	return PhysicsAsset->SkeletalBodySetups[Index];
+}
+
+USkeletalBodySetup* UPhysicsAssetToolset::FindBodyForShape(
+	UPhysicsAsset* PhysicsAsset, const FString& BoneName, const FString& ShapeName)
+{
+	if (!PhysicsAsset)
+	{
+		UKismetSystemLibrary::RaiseScriptError(TEXT("PhysicsAsset is null."));
+		return nullptr;
+	}
+	if (ShapeName.IsEmpty())
+	{
+		UKismetSystemLibrary::RaiseScriptError(TEXT("ShapeName cannot be empty."));
+		return nullptr;
+	}
+	USkeletalBodySetup* Body = FindBody(PhysicsAsset, BoneName);
+	if (!Body)
+	{
+		UKismetSystemLibrary::RaiseScriptError(
+			FString::Printf(TEXT("No body found for bone '%s'."), *BoneName));
+		return nullptr;
+	}
+	return Body;
+}
+
+FName UPhysicsAssetToolset::BeginShapeEdit(
+	UPhysicsAsset* PhysicsAsset, USkeletalBodySetup* Body, const FString& ShapeName)
+{
+	const FName Name(*ShapeName);
+	PhysicsAsset->Modify();
+	Body->Modify();
+	RemoveShapeByName(Body, Name);
+	return Name;
+}
+
+void UPhysicsAssetToolset::RemoveShapeByName(USkeletalBodySetup* Body, const FName& Name)
+{
+	Body->AggGeom.SphereElems.RemoveAll(
+		[&Name](const FKSphereElem& E) { return E.GetName() == Name; });
+	Body->AggGeom.SphylElems.RemoveAll(
+		[&Name](const FKSphylElem& E) { return E.GetName() == Name; });
+	Body->AggGeom.BoxElems.RemoveAll(
+		[&Name](const FKBoxElem& E) { return E.GetName() == Name; });
+}
+
+void UPhysicsAssetToolset::NotifyAssetChanged(UPhysicsAsset* PhysicsAsset)
+{
+	PhysicsAsset->UpdateBoundsBodiesArray();
+	PhysicsAsset->UpdateBodySetupIndexMap();
+	PhysicsAsset->MarkPackageDirty();
+#if WITH_EDITOR
+	PhysicsAsset->InvalidateAllPhysicsMeshes();
+	PhysicsAsset->RefreshPhysicsAssetChange();
+#endif
+}
+
+UPhysicsConstraintTemplate* UPhysicsAssetToolset::FindConstraintTemplate(
+	UPhysicsAsset* PhysicsAsset, const FString& Bone1Name, const FString& Bone2Name)
+{
+	const FName Name1(*Bone1Name);
+	const FName Name2(*Bone2Name);
+
+	int32 Index = PhysicsAsset->FindConstraintIndex(Name1, Name2);
+	if (Index == INDEX_NONE)
+	{
+		Index = PhysicsAsset->FindConstraintIndex(Name2, Name1);
+	}
+	if (Index == INDEX_NONE)
+	{
+		return nullptr;
+	}
+	return PhysicsAsset->ConstraintSetup[Index];
+}
diff --git a/Engine/Plugins/Experimental/Toolsets/PhysicsToolsets/Source/PhysicsToolsets/Private/PhysicsToolsets/PhysicsAssetToolset.h b/Engine/Plugins/Experimental/Toolsets/PhysicsToolsets/Source/PhysicsToolsets/Private/PhysicsToolsets/PhysicsAssetToolset.h
new file mode 100644
--- /dev/null
+++ b/Engine/Plugins/Experimental/Toolsets/PhysicsToolsets/Source/PhysicsToolsets/Private/PhysicsToolsets/PhysicsAssetToolset.h
@@ -0,0 +1,371 @@
+// Copyright Epic Games, Inc. All Rights Reserved.
+
+#pragma once
+
+#include "CoreMinimal.h"
+#include "ToolsetRegistry/UToolsetRegistry.h"
+
+#include "PhysicsAssetToolset.generated.h"
+
+class UPhysicsAsset;
+class UPhysicsConstraintTemplate;
+class USkeletalBodySetup;
+
+/** The type of collision primitive. */
+UENUM(BlueprintType)
+enum class EPhysicsShapeType : uint8
+{
+	Sphere,
+	Capsule,
+	Box,
+};
+
+/** Controls whether a physics body participates in world simulation. */
+UENUM(BlueprintType)
+enum class EBodyPhysicsMode : uint8
+{
+	/** Follows the skeletal mesh component's simulation state. */
+	Default,
+	/** Always kinematic: moves with animation and is not simulated. */
+	Kinematic,
+	/** Always simulated, regardless of the component's state. */
+	Simulated,
+};
+
+/** How a constraint axis may move. */
+UENUM(BlueprintType)
+enum class EConstraintMotion : uint8
+{
+	/** No restriction — the axis is completely free. */
+	Free,
+	/** Restricted to the specified limit angle. */
+	Limited,
+	/** Fully locked — no movement on this axis. */
+	Locked,
+};
+
+/**
+ * Describes a single collision primitive on a physics body.
+ * All positions and rotations are in the bone's local space.
+ */
+USTRUCT(BlueprintType)
+struct FPhysicsShapeInfo
+{
+	GENERATED_BODY()
+
+	/** User-defined name that uniquely identifies this shape on the body. */
+	UPROPERTY(BlueprintReadWrite, Category = "PhysicsShape")
+	FString ShapeName;
+
+	/** The type of collision primitive. */
+	UPROPERTY(BlueprintReadWrite, Category = "PhysicsShape")
+	EPhysicsShapeType ShapeType = EPhysicsShapeType::Sphere;
+
+	/** Center position in bone-local space (cm). */
+	UPROPERTY(BlueprintReadWrite, Category = "PhysicsShape")
+	FVector Center = FVector::ZeroVector;
+
+	/**
+	 * Orientation in bone-local space.
+	 * For Capsule, Z is the long axis. Unused for Sphere.
+	 */
+	UPROPERTY(BlueprintReadWrite, Category = "PhysicsShape")
+	FRotator Rotation = FRotator::ZeroRotator;
+
+	/** Radius (cm). Used by Sphere and Capsule. */
+	UPROPERTY(BlueprintReadWrite, Category = "PhysicsShape")
+	float Radius = 0.f;
+
+	/**
+	 * Cylinder section length (cm). Used by Capsule only.
+	 * Total capsule height = Length + 2 * Radius.
+	 */
+	UPROPERTY(BlueprintReadWrite, Category = "PhysicsShape")
+	float Length = 0.f;
+
+	/** Full extent along local X (cm). Used by Box only. */
+	UPROPERTY(BlueprintReadWrite, Category = "PhysicsShape")
+	float ExtentX = 0.f;
+
+	/** Full extent along local Y (cm). Used by Box only. */
+	UPROPERTY(BlueprintReadWrite, Category = "PhysicsShape")
+	float ExtentY = 0.f;
+
+	/** Full extent along local Z (cm). Used by Box only. */
+	UPROPERTY(BlueprintReadWrite, Category = "PhysicsShape")
+	float ExtentZ = 0.f;
+};
+
+/**
+ * Describes the angular limits for a single constraint between two bodies.
+ * Bone1 is the child bone; Bone2 is the parent bone.
+ * All limit angles are in degrees [0, 180].
+ */
+USTRUCT(BlueprintType)
+struct FPhysicsConstraintInfo
+{
+	GENERATED_BODY()
+
+	/** Name of the child bone. */
+	UPROPERTY(BlueprintReadWrite, Category = "PhysicsConstraint")
+	FString Bone1Name;
+
+	/** Name of the parent bone. */
+	UPROPERTY(BlueprintReadWrite, Category = "PhysicsConstraint")
+	FString Bone2Name;
+
+	/** Motion type for Swing 1 (rotation around the body-local Y axis). */
+	UPROPERTY(BlueprintReadWrite, Category = "PhysicsConstraint")
+	EConstraintMotion Swing1Motion = EConstraintMotion::Free;
+
+	/** Swing 1 limit in degrees. Used when Swing1Motion is Limited. */
+	UPROPERTY(BlueprintReadWrite, Category = "PhysicsConstraint")
+	float Swing1LimitDegrees = 45.f;
+
+	/** Motion type for Swing 2 (rotation around the body-local Z axis). */
+	UPROPERTY(BlueprintReadWrite, Category = "PhysicsConstraint")
+	EConstraintMotion Swing2Motion = EConstraintMotion::Free;
+
+	/** Swing 2 limit in degrees. Used when Swing2Motion is Limited. */
+	UPROPERTY(BlueprintReadWrite, Category = "PhysicsConstraint")
+	float Swing2LimitDegrees = 45.f;
+
+	/** Motion type for Twist (rotation around the body-local X axis). */
+	UPROPERTY(BlueprintReadWrite, Category = "PhysicsConstraint")
+	EConstraintMotion TwistMotion = EConstraintMotion::Free;
+
+	/** Twist limit in degrees. Used when TwistMotion is Limited. */
+	UPROPERTY(BlueprintReadWrite, Category = "PhysicsConstraint")
+	float TwistLimitDegrees = 45.f;
+};
+
+/** Provides tools for creating and managing Physics Assets. */
+UCLASS(BlueprintType, MinimalAPI)
+class UPhysicsAssetToolset : public UToolsetDefinition
+{
+	GENERATED_BODY()
+
+public:
+	/**
+	 * Creates a physics asset from a skeletal mesh, auto-generating collision bodies for
+	 * each bone. The asset is placed in the same folder as the mesh with the suffix
+	 * "_PhysicsAsset".
+	 * @param MeshPath Content-browser path to the skeletal mesh, e.g. '/Game/Characters/SKM_Hero'.
+	 * @param bAssignToMesh If true, assigns the new physics asset to the mesh.
+	 * @return The newly created physics asset.
+	 */
+	UFUNCTION(BlueprintCallable, Category = "PhysicsAsset")
+	static UPhysicsAsset* CreateFromMesh(const FString& MeshPath, bool bAssignToMesh);
+
+	/**
+	 * Returns the bone name for each rigid body in a physics asset.
+	 * @param PhysicsAsset The physics asset to query.
+	 * @return A list of bone names, one per body.
+	 */
+	UFUNCTION(BlueprintCallable, Category = "PhysicsAsset")
+	static TArray<FString> GetBodyNames(UPhysicsAsset* PhysicsAsset);
+
+	/**
+	 * Returns all collision shapes assigned to a body.
+	 * @param PhysicsAsset The physics asset to query.
+	 * @param BoneName The name of the bone whose body shapes to retrieve.
+	 * @return A list of shape descriptors. Raises a script error if no body exists for the bone.
+	 */
+	UFUNCTION(BlueprintCallable, Category = "PhysicsAsset")
+	static TArray<FPhysicsShapeInfo> GetBodyShapes(
+		UPhysicsAsset* PhysicsAsset, const FString& BoneName);
+
+	/**
+	 * Adds or replaces a sphere collision primitive on a body.
+	 * If any shape with the given name already exists on the body it is removed first.
+	 * @param PhysicsAsset The physics asset to modify.
+	 * @param BoneName The name of the bone whose body to modify.
+	 * @param ShapeName A name that uniquely identifies this shape on the body.
+	 * @param Center Center of the sphere in bone-local space (cm).
+	 * @param Radius Radius of the sphere (cm). Must be greater than zero.
+	 */
+	UFUNCTION(BlueprintCallable, Category = "PhysicsAsset")
+	static void SetSphere(
+		UPhysicsAsset* PhysicsAsset, const FString& BoneName,
+		const FString& ShapeName, const FVector& Center, float Radius);
+
+	/**
+	 * Adds or replaces a capsule collision primitive on a body.
+	 * If any shape with the given name already exists on the body it is removed first.
+	 * The capsule's long axis is its local Z after applying Rotation.
+	 * @param PhysicsAsset The physics asset to modify.
+	 * @param BoneName The name of the bone whose body to modify.
+	 * @param ShapeName A name that uniquely identifies this shape on the body.
+	 * @param Center Center of the capsule in bone-local space (cm).
+	 * @param Rotation Orientation of the capsule in bone-local space.
+	 * @param Radius Radius of the capsule end-caps (cm). Must be greater than zero.
+	 * @param Length Length of the cylindrical section (cm). Must be non-negative.
+	 *   Total capsule height = Length + 2 * Radius.
+	 */
+	UFUNCTION(BlueprintCallable, Category = "PhysicsAsset")
+	static void SetCapsule(
+		UPhysicsAsset* PhysicsAsset, const FString& BoneName,
+		const FString& ShapeName, const FVector& Center, const FRotator& Rotation,
+		float Radius, float Length);
+
+	/**
+	 * Adds or replaces a box collision primitive on a body.
+	 * If any shape with the given name already exists on the body it is removed first.
+	 * @param PhysicsAsset The physics asset to modify.
+	 * @param BoneName The name of the bone whose body to modify.
+	 * @param ShapeName A name that uniquely identifies this shape on the body.
+	 * @param Center Center of the box in bone-local space (cm).
+	 * @param Rotation Orientation of the box in bone-local space.
+	 * @param ExtentX Full extent along local X (cm). Must be greater than zero.
+	 * @param ExtentY Full extent along local Y (cm). Must be greater than zero.
+	 * @param ExtentZ Full extent along local Z (cm). Must be greater than zero.
+	 */
+	UFUNCTION(BlueprintCallable, Category = "PhysicsAsset")
+	static void SetBox(
+		UPhysicsAsset* PhysicsAsset, const FString& BoneName,
+		const FString& ShapeName, const FVector& Center, const FRotator& Rotation,
+		float ExtentX, float ExtentY, float ExtentZ);
+
+	/**
+	 * Removes a collision primitive from a body by name.
+	 * @param PhysicsAsset The physics asset to modify.
+	 * @param BoneName The name of the bone whose body to modify.
+	 * @param ShapeName The name of the shape to remove.
+	 */
+	UFUNCTION(BlueprintCallable, Category = "PhysicsAsset")
+	static void RemoveShape(
+		UPhysicsAsset* PhysicsAsset, const FString& BoneName, const FString& ShapeName);
+
+	// -------------------------------------------------------------------------
+	// Body CRUD
+	// -------------------------------------------------------------------------
+
+	/**
+	 * Adds a new empty body for the given bone.
+	 * @param PhysicsAsset The physics asset to modify.
+	 * @param BoneName The name of the bone to add a body for.
+	 */
+	UFUNCTION(BlueprintCallable, Category = "PhysicsAsset")
+	static void AddBody(UPhysicsAsset* PhysicsAsset, const FString& BoneName);
+
+	/**
+	 * Removes the body for the given bone along with any constraints that reference it.
+	 * Raises a script error if PhysicsAsset is null or no body exists for BoneName.
+	 * @param PhysicsAsset The physics asset to modify.
+	 * @param BoneName The name of the bone whose body to remove.
+	 */
+	UFUNCTION(BlueprintCallable, Category = "PhysicsAsset")
+	static void RemoveBody(UPhysicsAsset* PhysicsAsset, const FString& BoneName);
+
+	// -------------------------------------------------------------------------
+	// Body properties
+	// -------------------------------------------------------------------------
+
+	/**
+	 * Returns the physics simulation mode for the given body.
+	 * @param PhysicsAsset The physics asset to query.
+	 * @param BoneName The name of the bone whose body to query.
+	 */
+	UFUNCTION(BlueprintCallable, Category = "PhysicsAsset")
+	static EBodyPhysicsMode GetBodyPhysicsMode(
+		UPhysicsAsset* PhysicsAsset, const FString& BoneName);
+
+	/**
+	 * Sets the physics simulation mode for the given body.
+	 * @param PhysicsAsset The physics asset to modify.
+	 * @param BoneName The name of the bone whose body to modify.
+	 * @param Mode The desired simulation mode.
+	 */
+	UFUNCTION(BlueprintCallable, Category = "PhysicsAsset")
+	static void SetBodyPhysicsMode(
+		UPhysicsAsset* PhysicsAsset, const FString& BoneName, EBodyPhysicsMode Mode);
+
+	/**
+	 * Returns the mass-scale multiplier for the given body.
+	 * @param PhysicsAsset The physics asset to query.
+	 * @param BoneName The name of the bone whose body to query.
+	 */
+	UFUNCTION(BlueprintCallable, Category = "PhysicsAsset")
+	static float GetBodyMassScale(UPhysicsAsset* PhysicsAsset, const FString& BoneName);
+
+	/**
+	 * Sets the mass-scale multiplier for the given body.
+	 * @param PhysicsAsset The physics asset to modify.
+	 * @param BoneName The name of the bone whose body to modify.
+	 * @param MassScale Multiplier applied to the computed mass. Must be greater than zero.
+	 */
+	UFUNCTION(BlueprintCallable, Category = "PhysicsAsset")
+	static void SetBodyMassScale(
+		UPhysicsAsset* PhysicsAsset, const FString& BoneName, float MassScale);
+
+	// -------------------------------------------------------------------------
+	// Constraint CRUD
+	// -------------------------------------------------------------------------
+
+	/**
+	 * Returns all constraints in the physics asset with their current angular limits.
+	 * @param PhysicsAsset The physics asset to query.
+	 * @return A list of constraint descriptors.
+	 */
+	UFUNCTION(BlueprintCallable, Category = "PhysicsAsset")
+	static TArray<FPhysicsConstraintInfo> GetConstraints(UPhysicsAsset* PhysicsAsset);
+
+	/**
+	 * Adds a new constraint between two bodies. Both bodies must already exist.
+	 * @param PhysicsAsset The physics asset to modify.
+	 * @param Bone1Name Name of the child bone.
+	 * @param Bone2Name Name of the parent bone.
+	 */
+	UFUNCTION(BlueprintCallable, Category = "PhysicsAsset")
+	static void AddConstraint(
+		UPhysicsAsset* PhysicsAsset, const FString& Bone1Name, const FString& Bone2Name);
+
+	/**
+	 * Updates the angular limits for an existing constraint.
+	 * @param PhysicsAsset The physics asset to modify.
+	 * @param Info Constraint descriptor. Bone1Name and Bone2Name identify the constraint.
+	 */
+	UFUNCTION(BlueprintCallable, Category = "PhysicsAsset")
+	static void SetConstraintLimits(UPhysicsAsset* PhysicsAsset, FPhysicsConstraintInfo Info);
+
+	/**
+	 * Removes the constraint between two bodies.
+	 * @param PhysicsAsset The physics asset to modify.
+	 * @param Bone1Name Name of the child bone.
+	 * @param Bone2Name Name of the parent bone.
+	 */
+	UFUNCTION(BlueprintCallable, Category = "PhysicsAsset")
+	static void RemoveConstraint(
+		UPhysicsAsset* PhysicsAsset, const FString& Bone1Name, const FString& Bone2Name);
+
+private:
+	/** Returns the body for BoneName, or nullptr if not found. */
+	static USkeletalBodySetup* FindBody(UPhysicsAsset* PhysicsAsset, const FString& BoneName);
+
+	/** Removes all shapes named Name from Body across all primitive types. */
+	static void RemoveShapeByName(USkeletalBodySetup* Body, const FName& Name);
+
+	/**
+	 * Validates PhysicsAsset is non-null, ShapeName is non-empty, and a body exists for
+	 * BoneName. Raises a script error and returns nullptr on the first failed check.
+	 */
+	static USkeletalBodySetup* FindBodyForShape(
+		UPhysicsAsset* PhysicsAsset, const FString& BoneName, const FString& ShapeName);
+
+	/**
+	 * Marks PhysicsAsset and Body dirty, removes any existing shape named ShapeName,
+	 * and returns the FName to assign to the new element.
+	 */
+	static FName BeginShapeEdit(
+		UPhysicsAsset* PhysicsAsset, USkeletalBodySetup* Body, const FString& ShapeName);
+
+	/** Marks the asset dirty and triggers editor refresh. */
+	static void NotifyAssetChanged(UPhysicsAsset* PhysicsAsset);
+
+	/**
+	 * Returns the constraint template between the two bones, or nullptr if not found.
+	 * Tries both orderings of the bone names.
+	 */
+	static UPhysicsConstraintTemplate* FindConstraintTemplate(
+		UPhysicsAsset* PhysicsAsset, const FString& Bone1Name, const FString& Bone2Name);
+};
diff --git a/Engine/Plugins/Experimental/Toolsets/PhysicsToolsets/Source/PhysicsToolsets/Private/PhysicsToolsets/Tests/PhysicsAssetToolsetTest.cpp b/Engine/Plugins/Experimental/Toolsets/PhysicsToolsets/Source/PhysicsToolsets/Private/PhysicsToolsets/Tests/PhysicsAssetToolsetTest.cpp
new file mode 100644
--- /dev/null
+++ b/Engine/Plugins/Experimental/Toolsets/PhysicsToolsets/Source/PhysicsToolsets/Private/PhysicsToolsets/Tests/PhysicsAssetToolsetTest.cpp
@@ -0,0 +1,647 @@
+// Copyright Epic Games, Inc. All Rights Reserved.
+
+#include "CoreMinimal.h"
+#include "Misc/AutomationTest.h"
+#include "PhysicsEngine/PhysicsAsset.h"
+#include "PhysicsEngine/PhysicsConstraintTemplate.h"
+#include "PhysicsEngine/SkeletalBodySetup.h"
+#include "UObject/Package.h"
+#include "UObject/UObjectGlobals.h"
+
+#include "PhysicsToolsets/Tests/PhysicsToolsetsTestFlags.h"
+#include "PhysicsToolsets/PhysicsAssetToolset.h"
+
+#if WITH_DEV_AUTOMATION_TESTS
+
+BEGIN_DEFINE_SPEC(FPhysicsAssetToolsetSpec, "AI.PhysicsToolsets.PhysicsAssetToolsetSpec",
+	PhysicsToolsetsTest::Flags)
+
+	// In-memory physics asset shared across tests.
+	UPhysicsAsset* PhysicsAsset = nullptr;
+	FString BoneName;
+	FString BoneName2;
+
+	/** Creates a minimal in-memory physics asset with a single body named "TestBone". */
+	void SetUpInMemoryAsset()
+	{
+		PhysicsAsset = NewObject<UPhysicsAsset>(GetTransientPackage());
+		USkeletalBodySetup* Body = NewObject<USkeletalBodySetup>(PhysicsAsset);
+		Body->BoneName = FName("TestBone");
+		PhysicsAsset->SkeletalBodySetups.Add(Body);
+		PhysicsAsset->UpdateBodySetupIndexMap();
+		BoneName = TEXT("TestBone");
+	}
+
+	/** Extends SetUpInMemoryAsset with a second body named "TestBone2". */
+	void SetUpInMemoryAssetWithTwoBodies()
+	{
+		SetUpInMemoryAsset();
+		USkeletalBodySetup* Body2 = NewObject<USkeletalBodySetup>(PhysicsAsset);
+		Body2->BoneName = FName("TestBone2");
+		PhysicsAsset->SkeletalBodySetups.Add(Body2);
+		PhysicsAsset->UpdateBodySetupIndexMap();
+		BoneName2 = TEXT("TestBone2");
+	}
+
+	void TearDownInMemoryAsset()
+	{
+		PhysicsAsset = nullptr;
+		BoneName.Empty();
+		BoneName2.Empty();
+	}
+
+END_DEFINE_SPEC(FPhysicsAssetToolsetSpec)
+
+void FPhysicsAssetToolsetSpec::Define()
+{
+	// -------------------------------------------------------------------------
+	// CreateFromMesh
+	// -------------------------------------------------------------------------
+	Describe(TEXT("CreateFromMesh"), [this]()
+	{
+		It(TEXT("Returns null for an empty mesh path"), [this]()
+		{
+			UPhysicsAsset* Result = UPhysicsAssetToolset::CreateFromMesh(FString(), false);
+			TestNull("Null for empty path", Result);
+		});
+
+		It(TEXT("Returns null for an invalid mesh path"), [this]()
+		{
+			UPhysicsAsset* Result = UPhysicsAssetToolset::CreateFromMesh(
+				TEXT("/Game/ZZZ_DoesNotExist/NonExistentMesh"), false);
+			TestNull("Null for invalid path", Result);
+		});
+
+		It(TEXT("Creates a physics asset from a valid skeletal mesh"), [this]()
+		{
+			UPhysicsAsset* Result = UPhysicsAssetToolset::CreateFromMesh(
+				TEXT("/Engine/EditorMeshes/SkeletalMesh/DefaultSkeletalMesh"), false);
+			if (TestNotNull("Asset created", Result))
+			{
+				TestTrue("Has at least one body",
+					UPhysicsAssetToolset::GetBodyNames(Result).Num() > 0);
+			}
+		});
+	});
+
+	// -------------------------------------------------------------------------
+	// GetBodyNames
+	// -------------------------------------------------------------------------
+	Describe(TEXT("GetBodyNames"), [this]()
+	{
+		BeforeEach([this]() { SetUpInMemoryAsset(); });
+		AfterEach([this]() { TearDownInMemoryAsset(); });
+
+		It(TEXT("Returns the bone name of each body in the asset"), [this]()
+		{
+			const TArray<FString> Names =
+				UPhysicsAssetToolset::GetBodyNames(PhysicsAsset);
+			TestTrue("Contains TestBone", Names.Contains(BoneName));
+			TestEqual("Count matches body setups",
+				Names.Num(), PhysicsAsset->SkeletalBodySetups.Num());
+		});
+	});
+
+	// -------------------------------------------------------------------------
+	// GetBodyShapes
+	// -------------------------------------------------------------------------
+	Describe(TEXT("GetBodyShapes"), [this]()
+	{
+		BeforeEach([this]() { SetUpInMemoryAsset(); });
+		AfterEach([this]() { TearDownInMemoryAsset(); });
+
+		It(TEXT("Returns an empty list for a body with no shapes"), [this]()
+		{
+			const TArray<FPhysicsShapeInfo> Shapes =
+				UPhysicsAssetToolset::GetBodyShapes(PhysicsAsset, BoneName);
+			TestEqual("Empty", Shapes.Num(), 0);
+		});
+
+		It(TEXT("Returns an empty list for an unknown bone name"), [this]()
+		{
+			const TArray<FPhysicsShapeInfo> Shapes =
+				UPhysicsAssetToolset::GetBodyShapes(PhysicsAsset, TEXT("ZZZ_Unknown"));
+			TestEqual("Empty on unknown bone", Shapes.Num(), 0);
+		});
+	});
+
+	// -------------------------------------------------------------------------
+	// SetSphere
+	// -------------------------------------------------------------------------
+	Describe(TEXT("SetSphere"), [this]()
+	{
+		BeforeEach([this]() { SetUpInMemoryAsset(); });
+		AfterEach([this]() { TearDownInMemoryAsset(); });
+
+		It(TEXT("Adds a sphere visible via GetBodyShapes"), [this]()
+		{
+			UPhysicsAssetToolset::SetSphere(
+				PhysicsAsset, BoneName, TEXT("s1"), FVector(1.f, 2.f, 3.f), 10.f);
+
+			const TArray<FPhysicsShapeInfo> Shapes =
+				UPhysicsAssetToolset::GetBodyShapes(PhysicsAsset, BoneName);
+			if (TestEqual("One shape", Shapes.Num(), 1))
+			{
+				TestEqual("Sphere type", Shapes[0].ShapeType, EPhysicsShapeType::Sphere);
+				TestEqual("Name", Shapes[0].ShapeName, FString(TEXT("s1")));
+				TestEqual("Radius", Shapes[0].Radius, 10.f);
+				TestEqual("Center", Shapes[0].Center, FVector(1.f, 2.f, 3.f));
+			}
+		});
+
+		It(TEXT("Replaces an existing sphere that has the same name"), [this]()
+		{
+			UPhysicsAssetToolset::SetSphere(
+				PhysicsAsset, BoneName, TEXT("s1"), FVector::ZeroVector, 10.f);
+			UPhysicsAssetToolset::SetSphere(
+				PhysicsAsset, BoneName, TEXT("s1"), FVector::ZeroVector, 25.f);
+
+			const TArray<FPhysicsShapeInfo> Shapes =
+				UPhysicsAssetToolset::GetBodyShapes(PhysicsAsset, BoneName);
+			if (TestEqual("Still one shape", Shapes.Num(), 1))
+			{
+				TestEqual("Updated radius", Shapes[0].Radius, 25.f);
+			}
+		});
+
+		It(TEXT("Replaces a shape of a different type that has the same name"), [this]()
+		{
+			UPhysicsAssetToolset::SetBox(
+				PhysicsAsset, BoneName, TEXT("main"),
+				FVector::ZeroVector, FRotator::ZeroRotator, 10.f, 10.f, 10.f);
+			UPhysicsAssetToolset::SetSphere(
+				PhysicsAsset, BoneName, TEXT("main"), FVector::ZeroVector, 5.f);
+
+			const TArray<FPhysicsShapeInfo> Shapes =
+				UPhysicsAssetToolset::GetBodyShapes(PhysicsAsset, BoneName);
+			if (TestEqual("Still one shape", Shapes.Num(), 1))
+			{
+				TestEqual("Now a sphere", Shapes[0].ShapeType, EPhysicsShapeType::Sphere);
+			}
+		});
+
+		It(TEXT("Raises an error and adds nothing for zero radius"), [this]()
+		{
+			UPhysicsAssetToolset::SetSphere(
+				PhysicsAsset, BoneName, TEXT("s1"), FVector::ZeroVector, 0.f);
+			TestEqual("No shape added",
+				UPhysicsAssetToolset::GetBodyShapes(PhysicsAsset, BoneName).Num(), 0);
+		});
+
+		It(TEXT("Raises an error for an empty shape name"), [this]()
+		{
+			UPhysicsAssetToolset::SetSphere(
+				PhysicsAsset, BoneName, FString(), FVector::ZeroVector, 10.f);
+			TestEqual("No shape added",
+				UPhysicsAssetToolset::GetBodyShapes(PhysicsAsset, BoneName).Num(), 0);
+		});
+
+		It(TEXT("Raises an error for an unknown bone name"), [this]()
+		{
+			UPhysicsAssetToolset::SetSphere(
+				PhysicsAsset, TEXT("ZZZ_Unknown"), TEXT("s1"), FVector::ZeroVector, 10.f);
+		});
+	});
+
+	// -------------------------------------------------------------------------
+	// SetCapsule
+	// -------------------------------------------------------------------------
+	Describe(TEXT("SetCapsule"), [this]()
+	{
+		BeforeEach([this]() { SetUpInMemoryAsset(); });
+		AfterEach([this]() { TearDownInMemoryAsset(); });
+
+		It(TEXT("Adds a capsule visible via GetBodyShapes"), [this]()
+		{
+			UPhysicsAssetToolset::SetCapsule(
+				PhysicsAsset, BoneName, TEXT("c1"),
+				FVector::ZeroVector, FRotator::ZeroRotator, 5.f, 20.f);
+
+			const TArray<FPhysicsShapeInfo> Shapes =
+				UPhysicsAssetToolset::GetBodyShapes(PhysicsAsset, BoneName);
+			if (TestEqual("One shape", Shapes.Num(), 1))
+			{
+				TestEqual("Capsule type", Shapes[0].ShapeType, EPhysicsShapeType::Capsule);
+				TestEqual("Radius", Shapes[0].Radius, 5.f);
+				TestEqual("Length", Shapes[0].Length, 20.f);
+			}
+		});
+
+		It(TEXT("Raises an error for zero radius"), [this]()
+		{
+			UPhysicsAssetToolset::SetCapsule(
+				PhysicsAsset, BoneName, TEXT("c1"),
+				FVector::ZeroVector, FRotator::ZeroRotator, 0.f, 20.f);
+			TestEqual("No shape added",
+				UPhysicsAssetToolset::GetBodyShapes(PhysicsAsset, BoneName).Num(), 0);
+		});
+
+		It(TEXT("Raises an error for negative length"), [this]()
+		{
+			UPhysicsAssetToolset::SetCapsule(
+				PhysicsAsset, BoneName, TEXT("c1"),
+				FVector::ZeroVector, FRotator::ZeroRotator, 5.f, -1.f);
+			TestEqual("No shape added",
+				UPhysicsAssetToolset::GetBodyShapes(PhysicsAsset, BoneName).Num(), 0);
+		});
+
+		It(TEXT("Accepts zero length for a degenerate sphere-capped capsule"), [this]()
+		{
+			UPhysicsAssetToolset::SetCapsule(
+				PhysicsAsset, BoneName, TEXT("c1"),
+				FVector::ZeroVector, FRotator::ZeroRotator, 5.f, 0.f);
+			TestEqual("Shape added",
+				UPhysicsAssetToolset::GetBodyShapes(PhysicsAsset, BoneName).Num(), 1);
+		});
+	});
+
+	// -------------------------------------------------------------------------
+	// SetBox
+	// -------------------------------------------------------------------------
+	Describe(TEXT("SetBox"), [this]()
+	{
+		BeforeEach([this]() { SetUpInMemoryAsset(); });
+		AfterEach([this]() { TearDownInMemoryAsset(); });
+
+		It(TEXT("Adds a box visible via GetBodyShapes"), [this]()
+		{
+			UPhysicsAssetToolset::SetBox(
+				PhysicsAsset, BoneName, TEXT("b1"),
+				FVector::ZeroVector, FRotator::ZeroRotator, 10.f, 20.f, 30.f);
+
+			const TArray<FPhysicsShapeInfo> Shapes =
+				UPhysicsAssetToolset::GetBodyShapes(PhysicsAsset, BoneName);
+			if (TestEqual("One shape", Shapes.Num(), 1))
+			{
+				TestEqual("Box type", Shapes[0].ShapeType, EPhysicsShapeType::Box);
+				TestEqual("ExtentX", Shapes[0].ExtentX, 10.f);
+				TestEqual("ExtentY", Shapes[0].ExtentY, 20.f);
+				TestEqual("ExtentZ", Shapes[0].ExtentZ, 30.f);
+			}
+		});
+
+		It(TEXT("Raises an error for zero ExtentX"), [this]()
+		{
+			UPhysicsAssetToolset::SetBox(
+				PhysicsAsset, BoneName, TEXT("b1"),
+				FVector::ZeroVector, FRotator::ZeroRotator, 0.f, 10.f, 10.f);
+			TestEqual("No shape added",
+				UPhysicsAssetToolset::GetBodyShapes(PhysicsAsset, BoneName).Num(), 0);
+		});
+
+		It(TEXT("Raises an error for zero ExtentY"), [this]()
+		{
+			UPhysicsAssetToolset::SetBox(
+				PhysicsAsset, BoneName, TEXT("b1"),
+				FVector::ZeroVector, FRotator::ZeroRotator, 10.f, 0.f, 10.f);
+			TestEqual("No shape added",
+				UPhysicsAssetToolset::GetBodyShapes(PhysicsAsset, BoneName).Num(), 0);
+		});
+
+		It(TEXT("Raises an error for zero ExtentZ"), [this]()
+		{
+			UPhysicsAssetToolset::SetBox(
+				PhysicsAsset, BoneName, TEXT("b1"),
+				FVector::ZeroVector, FRotator::ZeroRotator, 10.f, 10.f, 0.f);
+			TestEqual("No shape added",
+				UPhysicsAssetToolset::GetBodyShapes(PhysicsAsset, BoneName).Num(), 0);
+		});
+	});
+
+	// -------------------------------------------------------------------------
+	// RemoveShape
+	// -------------------------------------------------------------------------
+	Describe(TEXT("RemoveShape"), [this]()
+	{
+		BeforeEach([this]() { SetUpInMemoryAsset(); });
+		AfterEach([this]() { TearDownInMemoryAsset(); });
+
+		It(TEXT("Removes a sphere by name"), [this]()
+		{
+			UPhysicsAssetToolset::SetSphere(
+				PhysicsAsset, BoneName, TEXT("s1"), FVector::ZeroVector, 10.f);
+			UPhysicsAssetToolset::RemoveShape(PhysicsAsset, BoneName, TEXT("s1"));
+			TestEqual("Empty after removal",
+				UPhysicsAssetToolset::GetBodyShapes(PhysicsAsset, BoneName).Num(), 0);
+		});
+
+		It(TEXT("Removes a capsule by name"), [this]()
+		{
+			UPhysicsAssetToolset::SetCapsule(
+				PhysicsAsset, BoneName, TEXT("c1"),
+				FVector::ZeroVector, FRotator::ZeroRotator, 5.f, 10.f);
+			UPhysicsAssetToolset::RemoveShape(PhysicsAsset, BoneName, TEXT("c1"));
+			TestEqual("Empty after removal",
+				UPhysicsAssetToolset::GetBodyShapes(PhysicsAsset, BoneName).Num(), 0);
+		});
+
+		It(TEXT("Removes a box by name"), [this]()
+		{
+			UPhysicsAssetToolset::SetBox(
+				PhysicsAsset, BoneName, TEXT("b1"),
+				FVector::ZeroVector, FRotator::ZeroRotator, 10.f, 10.f, 10.f);
+			UPhysicsAssetToolset::RemoveShape(PhysicsAsset, BoneName, TEXT("b1"));
+			TestEqual("Empty after removal",
+				UPhysicsAssetToolset::GetBodyShapes(PhysicsAsset, BoneName).Num(), 0);
+		});
+
+		It(TEXT("Leaves other shapes intact when removing one"), [this]()
+		{
+			UPhysicsAssetToolset::SetSphere(
+				PhysicsAsset, BoneName, TEXT("keep"), FVector::ZeroVector, 5.f);
+			UPhysicsAssetToolset::SetSphere(
+				PhysicsAsset, BoneName, TEXT("remove"), FVector::ZeroVector, 10.f);
+			UPhysicsAssetToolset::RemoveShape(PhysicsAsset, BoneName, TEXT("remove"));
+
+			const TArray<FPhysicsShapeInfo> Shapes =
+				UPhysicsAssetToolset::GetBodyShapes(PhysicsAsset, BoneName);
+			if (TestEqual("One shape remains", Shapes.Num(), 1))
+			{
+				TestEqual("Correct shape kept", Shapes[0].ShapeName, FString(TEXT("keep")));
+			}
+		});
+
+		It(TEXT("Raises an error for a non-existent shape name"), [this]()
+		{
+			UPhysicsAssetToolset::RemoveShape(
+				PhysicsAsset, BoneName, TEXT("ZZZ_NoSuchShape"));
+		});
+
+		It(TEXT("Raises an error for a non-existent bone name"), [this]()
+		{
+			UPhysicsAssetToolset::RemoveShape(
+				PhysicsAsset, TEXT("ZZZ_Unknown"), TEXT("s1"));
+		});
+	});
+
+	// -------------------------------------------------------------------------
+	// AddBody / RemoveBody
+	// -------------------------------------------------------------------------
+	Describe(TEXT("AddBody"), [this]()
+	{
+		BeforeEach([this]() { SetUpInMemoryAsset(); });
+		AfterEach([this]() { TearDownInMemoryAsset(); });
+
+		It(TEXT("Adds a body visible via GetBodyNames"), [this]()
+		{
+			UPhysicsAssetToolset::AddBody(PhysicsAsset, TEXT("NewBone"));
+			const TArray<FString> Names = UPhysicsAssetToolset::GetBodyNames(PhysicsAsset);
+			TestTrue("NewBone present", Names.Contains(TEXT("NewBone")));
+		});
+
+		It(TEXT("Raises an error for an empty bone name"), [this]()
+		{
+			const int32 CountBefore = UPhysicsAssetToolset::GetBodyNames(PhysicsAsset).Num();
+			UPhysicsAssetToolset::AddBody(PhysicsAsset, FString());
+			TestEqual("Count unchanged",
+				UPhysicsAssetToolset::GetBodyNames(PhysicsAsset).Num(), CountBefore);
+		});
+
+		It(TEXT("Raises an error when the bone already has a body"), [this]()
+		{
+			const int32 CountBefore = UPhysicsAssetToolset::GetBodyNames(PhysicsAsset).Num();
+			UPhysicsAssetToolset::AddBody(PhysicsAsset, BoneName);
+			TestEqual("Count unchanged",
+				UPhysicsAssetToolset::GetBodyNames(PhysicsAsset).Num(), CountBefore);
+		});
+	});
+
+	Describe(TEXT("RemoveBody"), [this]()
+	{
+		BeforeEach([this]() { SetUpInMemoryAssetWithTwoBodies(); });
+		AfterEach([this]() { TearDownInMemoryAsset(); });
+
+		It(TEXT("Removes a body by bone name"), [this]()
+		{
+			UPhysicsAssetToolset::RemoveBody(PhysicsAsset, BoneName);
+			const TArray<FString> Names = UPhysicsAssetToolset::GetBodyNames(PhysicsAsset);
+			TestFalse("No longer present", Names.Contains(BoneName));
+		});
+
+		It(TEXT("Also removes constraints that reference the body"), [this]()
+		{
+			UPhysicsAssetToolset::AddConstraint(PhysicsAsset, BoneName, BoneName2);
+			UPhysicsAssetToolset::RemoveBody(PhysicsAsset, BoneName);
+			TestEqual("No constraints remain",
+				UPhysicsAssetToolset::GetConstraints(PhysicsAsset).Num(), 0);
+		});
+
+		It(TEXT("Raises an error for a non-existent bone name"), [this]()
+		{
+			const int32 CountBefore = UPhysicsAssetToolset::GetBodyNames(PhysicsAsset).Num();
+			UPhysicsAssetToolset::RemoveBody(PhysicsAsset, TEXT("ZZZ_Unknown"));
+			TestEqual("Count unchanged",
+				UPhysicsAssetToolset::GetBodyNames(PhysicsAsset).Num(), CountBefore);
+		});
+	});
+
+	// -------------------------------------------------------------------------
+	// GetBodyPhysicsMode / SetBodyPhysicsMode
+	// -------------------------------------------------------------------------
+	Describe(TEXT("GetBodyPhysicsMode and SetBodyPhysicsMode"), [this]()
+	{
+		BeforeEach([this]() { SetUpInMemoryAsset(); });
+		AfterEach([this]() { TearDownInMemoryAsset(); });
+
+		It(TEXT("Default mode is Default"), [this]()
+		{
+			TestEqual("Default",
+				UPhysicsAssetToolset::GetBodyPhysicsMode(PhysicsAsset, BoneName),
+				EBodyPhysicsMode::Default);
+		});
+
+		It(TEXT("Can set and read back Kinematic"), [this]()
+		{
+			UPhysicsAssetToolset::SetBodyPhysicsMode(
+				PhysicsAsset, BoneName, EBodyPhysicsMode::Kinematic);
+			TestEqual("Kinematic",
+				UPhysicsAssetToolset::GetBodyPhysicsMode(PhysicsAsset, BoneName),
+				EBodyPhysicsMode::Kinematic);
+		});
+
+		It(TEXT("Can set and read back Simulated"), [this]()
+		{
+			UPhysicsAssetToolset::SetBodyPhysicsMode(
+				PhysicsAsset, BoneName, EBodyPhysicsMode::Simulated);
+			TestEqual("Simulated",
+				UPhysicsAssetToolset::GetBodyPhysicsMode(PhysicsAsset, BoneName),
+				EBodyPhysicsMode::Simulated);
+		});
+
+		It(TEXT("Raises an error for an unknown bone name"), [this]()
+		{
+			UPhysicsAssetToolset::SetBodyPhysicsMode(
+				PhysicsAsset, TEXT("ZZZ_Unknown"), EBodyPhysicsMode::Simulated);
+		});
+	});
+
+	// -------------------------------------------------------------------------
+	// GetBodyMassScale / SetBodyMassScale
+	// -------------------------------------------------------------------------
+	Describe(TEXT("GetBodyMassScale and SetBodyMassScale"), [this]()
+	{
+		BeforeEach([this]() { SetUpInMemoryAsset(); });
+		AfterEach([this]() { TearDownInMemoryAsset(); });
+
+		It(TEXT("Can set and read back a mass scale"), [this]()
+		{
+			UPhysicsAssetToolset::SetBodyMassScale(PhysicsAsset, BoneName, 2.5f);
+			TestEqual("MassScale",
+				UPhysicsAssetToolset::GetBodyMassScale(PhysicsAsset, BoneName), 2.5f);
+		});
+
+		It(TEXT("Raises an error for zero mass scale"), [this]()
+		{
+			UPhysicsAssetToolset::SetBodyMassScale(PhysicsAsset, BoneName, 0.f);
+		});
+
+		It(TEXT("Raises an error for negative mass scale"), [this]()
+		{
+			UPhysicsAssetToolset::SetBodyMassScale(PhysicsAsset, BoneName, -1.f);
+		});
+
+		It(TEXT("Raises an error for an unknown bone name"), [this]()
+		{
+			UPhysicsAssetToolset::SetBodyMassScale(
+				PhysicsAsset, TEXT("ZZZ_Unknown"), 1.f);
+		});
+	});
+
+	// -------------------------------------------------------------------------
+	// GetConstraints / AddConstraint / SetConstraintLimits / RemoveConstraint
+	// -------------------------------------------------------------------------
+	Describe(TEXT("GetConstraints"), [this]()
+	{
+		BeforeEach([this]() { SetUpInMemoryAssetWithTwoBodies(); });
+		AfterEach([this]() { TearDownInMemoryAsset(); });
+
+		It(TEXT("Returns an empty list when there are no constraints"), [this]()
+		{
+			TestEqual("Empty",
+				UPhysicsAssetToolset::GetConstraints(PhysicsAsset).Num(), 0);
+		});
+
+		It(TEXT("Returns one entry per constraint"), [this]()
+		{
+			UPhysicsAssetToolset::AddConstraint(PhysicsAsset, BoneName, BoneName2);
+			TestEqual("One constraint",
+				UPhysicsAssetToolset::GetConstraints(PhysicsAsset).Num(), 1);
+		});
+	});
+
+	Describe(TEXT("AddConstraint"), [this]()
+	{
+		BeforeEach([this]() { SetUpInMemoryAssetWithTwoBodies(); });
+		AfterEach([this]() { TearDownInMemoryAsset(); });
+
+		It(TEXT("Adds a constraint visible via GetConstraints"), [this]()
+		{
+			UPhysicsAssetToolset::AddConstraint(PhysicsAsset, BoneName, BoneName2);
+			const TArray<FPhysicsConstraintInfo> Constraints =
+				UPhysicsAssetToolset::GetConstraints(PhysicsAsset);
+			if (TestEqual("One constraint", Constraints.Num(), 1))
+			{
+				TestEqual("Bone1", Constraints[0].Bone1Name, BoneName);
+				TestEqual("Bone2", Constraints[0].Bone2Name, BoneName2);
+			}
+		});
+
+		It(TEXT("Raises an error when a constraint already exists"), [this]()
+		{
+			UPhysicsAssetToolset::AddConstraint(PhysicsAsset, BoneName, BoneName2);
+			UPhysicsAssetToolset::AddConstraint(PhysicsAsset, BoneName, BoneName2);
+			TestEqual("Still one constraint",
+				UPhysicsAssetToolset::GetConstraints(PhysicsAsset).Num(), 1);
+		});
+
+		It(TEXT("Raises an error for empty bone names"), [this]()
+		{
+			UPhysicsAssetToolset::AddConstraint(PhysicsAsset, FString(), BoneName2);
+			UPhysicsAssetToolset::AddConstraint(PhysicsAsset, BoneName, FString());
+			TestEqual("No constraints added",
+				UPhysicsAssetToolset::GetConstraints(PhysicsAsset).Num(), 0);
+		});
+
+		It(TEXT("Raises an error when a bone has no body"), [this]()
+		{
+			UPhysicsAssetToolset::AddConstraint(
+				PhysicsAsset, BoneName, TEXT("ZZZ_Unknown"));
+			TestEqual("No constraint added",
+				UPhysicsAssetToolset::GetConstraints(PhysicsAsset).Num(), 0);
+		});
+	});
+
+	Describe(TEXT("SetConstraintLimits"), [this]()
+	{
+		BeforeEach([this]()
+		{
+			SetUpInMemoryAssetWithTwoBodies();
+			UPhysicsAssetToolset::AddConstraint(PhysicsAsset, BoneName, BoneName2);
+		});
+		AfterEach([this]() { TearDownInMemoryAsset(); });
+
+		It(TEXT("Sets angular limits and reads them back via GetConstraints"), [this]()
+		{
+			FPhysicsConstraintInfo Info;
+			Info.Bone1Name = BoneName;
+			Info.Bone2Name = BoneName2;
+			Info.Swing1Motion = EConstraintMotion::Limited;
+			Info.Swing1LimitDegrees = 30.f;
+			Info.Swing2Motion = EConstraintMotion::Locked;
+			Info.Swing2LimitDegrees = 0.f;
+			Info.TwistMotion = EConstraintMotion::Limited;
+			Info.TwistLimitDegrees = 45.f;
+
+			UPhysicsAssetToolset::SetConstraintLimits(PhysicsAsset, Info);
+
+			const TArray<FPhysicsConstraintInfo> Constraints =
+				UPhysicsAssetToolset::GetConstraints(PhysicsAsset);
+			if (TestEqual("One constraint", Constraints.Num(), 1))
+			{
+				TestEqual("Swing1Motion", Constraints[0].Swing1Motion,
+					EConstraintMotion::Limited);
+				TestEqual("Swing1Limit", Constraints[0].Swing1LimitDegrees, 30.f);
+				TestEqual("Swing2Motion", Constraints[0].Swing2Motion,
+					EConstraintMotion::Locked);
+				TestEqual("TwistMotion", Constraints[0].TwistMotion,
+					EConstraintMotion::Limited);
+				TestEqual("TwistLimit", Constraints[0].TwistLimitDegrees, 45.f);
+			}
+		});
+
+		It(TEXT("Raises an error when the constraint does not exist"), [this]()
+		{
+			FPhysicsConstraintInfo Info;
+			Info.Bone1Name = TEXT("ZZZ_A");
+			Info.Bone2Name = TEXT("ZZZ_B");
+			UPhysicsAssetToolset::SetConstraintLimits(PhysicsAsset, Info);
+		});
+	});
+
+	Describe(TEXT("RemoveConstraint"), [this]()
+	{
+		BeforeEach([this]()
+		{
+			SetUpInMemoryAssetWithTwoBodies();
+			UPhysicsAssetToolset::AddConstraint(PhysicsAsset, BoneName, BoneName2);
+		});
+		AfterEach([this]() { TearDownInMemoryAsset(); });
+
+		It(TEXT("Removes an existing constraint"), [this]()
+		{
+			UPhysicsAssetToolset::RemoveConstraint(PhysicsAsset, BoneName, BoneName2);
+			TestEqual("Empty after removal",
+				UPhysicsAssetToolset::GetConstraints(PhysicsAsset).Num(), 0);
+		});
+
+		It(TEXT("Raises an error for a non-existent constraint"), [this]()
+		{
+			UPhysicsAssetToolset::RemoveConstraint(
+				PhysicsAsset, TEXT("ZZZ_A"), TEXT("ZZZ_B"));
+			TestEqual("Existing constraint unaffected",
+				UPhysicsAssetToolset::GetConstraints(PhysicsAsset).Num(), 1);
+		});
+	});
+}
+
+#endif  // WITH_DEV_AUTOMATION_TESTS
diff --git a/Engine/Plugins/Experimental/Toolsets/PhysicsToolsets/Source/PhysicsToolsets/Private/PhysicsToolsets/Tests/PhysicsToolsetsTestFlags.h b/Engine/Plugins/Experimental/Toolsets/PhysicsToolsets/Source/PhysicsToolsets/Private/PhysicsToolsets/Tests/PhysicsToolsetsTestFlags.h
new file mode 100644
--- /dev/null
+++ b/Engine/Plugins/Experimental/Toolsets/PhysicsToolsets/Source/PhysicsToolsets/Private/PhysicsToolsets/Tests/PhysicsToolsetsTestFlags.h
@@ -0,0 +1,17 @@
+// Copyright Epic Games, Inc. All Rights Reserved.
+#pragma once
+
+#include "Misc/AutomationTest.h"
+
+#if WITH_DEV_AUTOMATION_TESTS
+
+namespace PhysicsToolsetsTest
+{
+	// Default flags to use for tests.
+	const auto Flags =
+		EAutomationTestFlags::EditorContext |
+		EAutomationTestFlags::ProductFilter |
+		EAutomationTestFlags::CriticalPriority;
+}  // namespace PhysicsToolsetsTest
+
+#endif  // WITH_DEV_AUTOMATION_TESTS
